UE-43240 Removing dependency on component visualizers in runtime Phonon module.

diff --git a/Engine/Plugins/Runtime/Phonon/Source/Phonon/Phonon.Build.cs b/Engine/Plugins/Runtime/Phonon/Source/Phonon/Phonon.Build.cs
--- a/Engine/Plugins/Runtime/Phonon/Source/Phonon/Phonon.Build.cs
+++ b/Engine/Plugins/Runtime/Phonon/Source/Phonon/Phonon.Build.cs
@@ -31,7 +31,6 @@ namespace UnrealBuildTool.Rules
             PrivateIncludePathModuleNames.AddRange(
                 new string[] {
                     "TargetPlatform",
-                    "ComponentVisualizers"
                 }
             );
 
@@ -44,7 +43,6 @@ namespace UnrealBuildTool.Rules
                     "Projects",
                     "AudioMixer",
                     "XAudio2",
-                    "ComponentVisualizers"
                 }
             );
 
diff --git a/Engine/Plugins/Runtime/Phonon/Source/PhononEditor/PhononEditor.Build.cs b/Engine/Plugins/Runtime/Phonon/Source/PhononEditor/PhononEditor.Build.cs
--- a/Engine/Plugins/Runtime/Phonon/Source/PhononEditor/PhononEditor.Build.cs
+++ b/Engine/Plugins/Runtime/Phonon/Source/PhononEditor/PhononEditor.Build.cs
@@ -60,7 +60,8 @@ namespace UnrealBuildTool.Rules
                     "PropertyEditor",
                     "Projects",
                     "EditorStyle",
-                    "Phonon"
+                    "Phonon",
+                    "ComponentVisualizers",
                  }
             );
 
