Fix phonon non-editor compilation

diff --git a/Engine/Plugins/Runtime/Phonon/Source/Phonon/Private/PhononListenerComponent.cpp b/Engine/Plugins/Runtime/Phonon/Source/Phonon/Private/PhononListenerComponent.cpp
--- a/Engine/Plugins/Runtime/Phonon/Source/Phonon/Private/PhononListenerComponent.cpp
+++ b/Engine/Plugins/Runtime/Phonon/Source/Phonon/Private/PhononListenerComponent.cpp
@@ -3,9 +3,6 @@
 //
 
 #include "PhononListenerComponent.h"
-#include "LevelEditorViewport.h"
-#include "Settings/LevelEditorPlaySettings.h"
-#include "Settings/LevelEditorViewportSettings.h"
 #include "PhononModule.h"
 #include "PhononOcclusion.h"
 #include "PhononReverb.h"
diff --git a/Engine/Plugins/Runtime/Phonon/Source/Phonon/Private/PhononProbeVolume.cpp b/Engine/Plugins/Runtime/Phonon/Source/Phonon/Private/PhononProbeVolume.cpp
--- a/Engine/Plugins/Runtime/Phonon/Source/Phonon/Private/PhononProbeVolume.cpp
+++ b/Engine/Plugins/Runtime/Phonon/Source/Phonon/Private/PhononProbeVolume.cpp
@@ -12,10 +12,6 @@
 #include "Editor.h"
 #endif
 
-#include "LevelEditorViewport.h"
-#include "Settings/LevelEditorPlaySettings.h"
-#include "Settings/LevelEditorViewportSettings.h"
-
 #include <algorithm>
 
 APhononProbeVolume::APhononProbeVolume(const FObjectInitializer& ObjectInitializer)
