Proper fix for Phonon compile errors

diff --git a/Engine/Plugins/Runtime/Phonon/Source/Phonon/Private/PhononProbeVolume.cpp b/Engine/Plugins/Runtime/Phonon/Source/Phonon/Private/PhononProbeVolume.cpp
--- a/Engine/Plugins/Runtime/Phonon/Source/Phonon/Private/PhononProbeVolume.cpp
+++ b/Engine/Plugins/Runtime/Phonon/Source/Phonon/Private/PhononProbeVolume.cpp
@@ -10,6 +10,7 @@
 
 #if WITH_EDITOR
 #include "Editor.h"
+#include "LevelEditorViewport.h"
 #endif
 
 #include <algorithm>
